Add missing AutomationTool.Tests project.

#rnx

diff --git a/Engine/Source/Programs/AutomationTool.Tests/AutomationTool.Tests.csproj b/Engine/Source/Programs/AutomationTool.Tests/AutomationTool.Tests.csproj
new file mode 100644
--- /dev/null
+++ b/Engine/Source/Programs/AutomationTool.Tests/AutomationTool.Tests.csproj
@@ -0,0 +1,26 @@
+<Project Sdk="Microsoft.NET.Sdk">
+
+  <PropertyGroup>
+    <TargetFramework>net6.0</TargetFramework>
+    <ImplicitUsings>enable</ImplicitUsings>
+    <Nullable>enable</Nullable>
+
+    <IsPackable>false</IsPackable>
+    <IsTestProject>true</IsTestProject>
+  </PropertyGroup>
+
+  <ItemGroup>
+    <PackageReference Include="Microsoft.NET.Test.Sdk" Version="17.5.0" />
+    <PackageReference Include="MSTest.TestAdapter" Version="2.2.10" />
+    <PackageReference Include="MSTest.TestFramework" Version="2.2.10" />
+    <PackageReference Include="coverlet.collector" Version="3.2.0" />
+  </ItemGroup>
+
+  <ItemGroup>
+    <ProjectReference Include="..\AutomationTool\AutomationUtils\AutomationUtils.Automation.csproj" />
+    <ProjectReference Include="..\Shared\EpicGames.Core\EpicGames.Core.csproj" />
+    <ProjectReference Include="..\Shared\EpicGames.Perforce\EpicGames.Perforce.csproj" />
+    <ProjectReference Include="..\UnrealBuildTool\UnrealBuildTool.csproj" />
+  </ItemGroup>
+
+</Project>
