Horde: Match MSTest errors which contain parameter lists.

#rnx

diff --git a/Engine/Source/Programs/AutomationTool.Tests/LogParserTests.cs b/Engine/Source/Programs/AutomationTool.Tests/LogParserTests.cs
--- a/Engine/Source/Programs/AutomationTool.Tests/LogParserTests.cs
+++ b/Engine/Source/Programs/AutomationTool.Tests/LogParserTests.cs
@@ -10,7 +10,7 @@ using EpicGames.Perforce;
 using Microsoft.Extensions.Logging;
 using Microsoft.VisualStudio.TestTools.UnitTesting;
 
-namespace Horde.Agent.Tests
+namespace AutomationTool.Tests
 {
 	[TestClass]
 	public class LogParserTests
@@ -830,6 +830,26 @@ namespace Horde.Agent.Tests
 			CheckEventGroup(logEvents, 2, 5, LogLevel.Error, KnownLogEvents.MSTest);
 		}
 
+		[TestMethod]
+		public void MsTestEventMatcher3()
+		{
+			string[] lines =
+			{
+				@" info: Horde.Server.Auditing.AuditLog[0]",
+				@"       Session 662a561ef41fb04ceda48319 started",
+				@"",
+				@"  Failed AutoConformAgentsAsync (True,200,) [208 ms]",
+				@"  Error Message:",
+				@"   Assert.AreEqual failed. Expected:<True>. Actual:<False>.",
+				@"  Stack Trace:",
+				@"     at Horde.Server.Tests.Agents.Pools.PoolUpdateServiceTest.AutoConformAgentsAsync(Boolean conformRequested, Nullable`1[] autoConformThresholdsM) in /app/Source/Programs/Horde/Horde.Server.Tests/Agents/Pools/PoolUpdateServiceTest.cs:line 120",
+			};
+
+			List<LogEvent> logEvents = Parse(lines);
+			Assert.AreEqual(5, logEvents.Count);
+			CheckEventGroup(logEvents, 3, 5, LogLevel.Error, KnownLogEvents.MSTest);
+		}
+
 		static List<LogEvent> Parse(IEnumerable<string> lines)
 		{
 			return Parse(String.Join("\n", lines));
diff --git a/Engine/Source/Programs/AutomationTool/AutomationUtils/Matchers/MsTestEventMatcher.cs b/Engine/Source/Programs/AutomationTool/AutomationUtils/Matchers/MsTestEventMatcher.cs
--- a/Engine/Source/Programs/AutomationTool/AutomationUtils/Matchers/MsTestEventMatcher.cs
+++ b/Engine/Source/Programs/AutomationTool/AutomationUtils/Matchers/MsTestEventMatcher.cs
@@ -10,7 +10,7 @@ namespace AutomationUtils.Matchers
 {
 	class MsTestEventMatcher : ILogEventMatcher
 	{
-		static readonly Regex s_failedPattern = new Regex(@"^  \s*Failed [A-Za-z0-9_]+ \[\d+ (?:s|ms)\]$");
+		static readonly Regex s_failedPattern = new Regex(@"^  \s*Failed [A-Za-z0-9_]+ (?:\(.*\) )?\[\d+ (?:s|ms)\]$");
 		static readonly Regex s_detailPattern = new Regex(@"^  \s*(?:Error Message|Stack Trace):$");
 
 		public LogEventMatch? Match(ILogCursor cursor)
