Fix two issues in ThreadSanitizerEventMatcher:
- We could match on irrelevant log line and then be forced to return a partial TSAN report depite not finding our report end delimiter
- We could potentially move the LogCursor one more line than available in the buffer if we end parsing on a line that matches our stack trace pattern since we would call MoveNext() on potentially a null line. This would result in exceptions being thrown while generating events.

diff --git a/Engine/Source/Programs/AutomationTool.Tests/LogParserTests.cs b/Engine/Source/Programs/AutomationTool.Tests/LogParserTests.cs
--- a/Engine/Source/Programs/AutomationTool.Tests/LogParserTests.cs
+++ b/Engine/Source/Programs/AutomationTool.Tests/LogParserTests.cs
@@ -699,6 +699,27 @@ namespace AutomationTool.Tests
 			}
 		}
 
+		[TestMethod]
+		public void ThreadSanitizerErrorMatcher_NoMatch()
+		{
+			string[] lines =
+			{
+				@"==================", 
+				@"WARNING: ThreadSanitizer: data race (pid=21089)",
+				@"  Write of size 1 at 0x00004060ea38 by thread T25:",
+				@"    #0 FPaths::IsStaged() /mnt/horde/++UE5/Sync/Engine/Source/./Runtime/Core/Private/Misc/Paths.cpp:172:33 (CitySampleEditor+0x2c0f2f44) (BuildId: 6622c84dbb6a946e)",
+				@"",
+				@"  Previous write of size 1 at 0x00004060ea38 by thread T20:",
+				@"    #0 FPaths::IsStaged() /mnt/horde/++UE5/Sync/Engine/Source/./Runtime/Core/Private/Misc/Paths.cpp:172:33 (CitySampleEditor+0x2c0f2f44) (BuildId: 6622c84dbb6a946e)",
+			};
+
+			// We won't find a terminating string for the TSAN report so we only log a source file warning
+			List<LogEvent> logEvents = Parse(lines);
+			Assert.AreEqual(1, logEvents.Count);
+			Assert.AreEqual(LogLevel.Warning, logEvents[0].Level);
+			Assert.AreEqual(KnownLogEvents.AutomationTool_SourceFileLine, logEvents[0].Id);
+		}
+
 		[TestMethod]
 		public void ThreadSanitizerErrorMatcher()
 		{
diff --git a/Engine/Source/Programs/AutomationTool/AutomationUtils/Matchers/ThreadSanitizerEventMatcher.cs b/Engine/Source/Programs/AutomationTool/AutomationUtils/Matchers/ThreadSanitizerEventMatcher.cs
--- a/Engine/Source/Programs/AutomationTool/AutomationUtils/Matchers/ThreadSanitizerEventMatcher.cs
+++ b/Engine/Source/Programs/AutomationTool/AutomationUtils/Matchers/ThreadSanitizerEventMatcher.cs
@@ -75,10 +75,16 @@ namespace AutomationUtils.Matchers
 						builder.Annotate(match.Groups["Line"], LogEventMarkup.LineNumber);
 						builder.Annotate(match.Groups["Column"], LogEventMarkup.ColumnNumber);
 						builder.AnnotateSymbol(match.Groups["Symbol"]);
+
+						builder.MoveNext();
+						return builder.ToMatch(LogEventPriority.Normal, LogLevel.Error, KnownLogEvents.Sanitizer_Thread);
+					}
+
+					if (builder.Current.CurrentLine != null)
+					{
+						builder.MoveNext();
 					}
-					builder.MoveNext();
 				}
-				return builder.ToMatch(LogEventPriority.Normal, LogLevel.Error, KnownLogEvents.Sanitizer_Thread);
 			}
 
 			return null;
