Horde: Fix inconsistent nullability annotations in LogEvent/MessageTemplate classes, and add tests for round tripping log events to/from JSON.

diff --git a/Engine/Source/Programs/AutomationTool.Tests/LogParserTests.cs b/Engine/Source/Programs/AutomationTool.Tests/LogParserTests.cs
--- a/Engine/Source/Programs/AutomationTool.Tests/LogParserTests.cs
+++ b/Engine/Source/Programs/AutomationTool.Tests/LogParserTests.cs
@@ -36,7 +36,7 @@ namespace AutomationTool.Tests
 				LogEvent logEvent = LogEvent.Read(JsonLogEvent.FromLoggerState(logLevel, eventId, state, exception, formatter).Data.Span);
 				if (logEvent.Level != LogLevel.Information || logEvent.Id != default || logEvent.Properties != null)
 				{
-					KeyValuePair<string, object>[] items = new[] { new KeyValuePair<string, object>(LogLine, _logLineIndex) };
+					KeyValuePair<string, object?>[] items = new[] { new KeyValuePair<string, object?>(LogLine, _logLineIndex) };
 					logEvent.Properties = (logEvent.Properties == null) ? items : Enumerable.Concat(logEvent.Properties, items);
 					_events.Add(logEvent);
 				}
@@ -220,8 +220,8 @@ namespace AutomationTool.Tests
 				LogValue fileProperty1 = (LogValue)logEvents[0].GetProperty("file");
 				Assert.AreEqual(@"GenerateSigningRequestDialog.cs", fileProperty1.Text);
 				Assert.AreEqual(@"SourceFile", fileProperty1.Type.ToString());
-				Assert.AreEqual(@"Engine/Source/Programs/IOS/iPhonePackager/GenerateSigningRequestDialog.cs", fileProperty1.Properties![LogEventPropertyName.RelativePath].ToString());
-				Assert.AreEqual(@"//UE4/Main/Engine/Source/Programs/IOS/iPhonePackager/GenerateSigningRequestDialog.cs@12345", fileProperty1.Properties[LogEventPropertyName.DepotPath].ToString());
+				Assert.AreEqual(@"Engine/Source/Programs/IOS/iPhonePackager/GenerateSigningRequestDialog.cs", fileProperty1.Properties![LogEventPropertyName.RelativePath]?.ToString());
+				Assert.AreEqual(@"//UE4/Main/Engine/Source/Programs/IOS/iPhonePackager/GenerateSigningRequestDialog.cs@12345", fileProperty1.Properties[LogEventPropertyName.DepotPath]?.ToString());
 
 				CheckEventGroup(logEvents.Slice(1, 1), 1, 1, LogLevel.Error, KnownLogEvents.Compiler);
 				Assert.AreEqual("CS0246", logEvents[1].GetProperty("code").ToString());
@@ -230,8 +230,8 @@ namespace AutomationTool.Tests
 				LogValue fileProperty2 = (LogValue)logEvents[1].GetProperty("file");
 				Assert.AreEqual(@"Utilities.cs", fileProperty2.Text);
 				Assert.AreEqual(@"SourceFile", fileProperty2.Type.ToString());
-				Assert.AreEqual(@"Engine/Source/Programs/IOS/iPhonePackager/Utilities.cs", fileProperty2.Properties![LogEventPropertyName.RelativePath].ToString());
-				Assert.AreEqual(@"//UE4/Main/Engine/Source/Programs/IOS/iPhonePackager/Utilities.cs@12345", fileProperty2.Properties[LogEventPropertyName.DepotPath].ToString());
+				Assert.AreEqual(@"Engine/Source/Programs/IOS/iPhonePackager/Utilities.cs", fileProperty2.Properties![LogEventPropertyName.RelativePath]?.ToString());
+				Assert.AreEqual(@"//UE4/Main/Engine/Source/Programs/IOS/iPhonePackager/Utilities.cs@12345", fileProperty2.Properties[LogEventPropertyName.DepotPath]?.ToString());
 			}
 		}
 
@@ -254,8 +254,8 @@ namespace AutomationTool.Tests
 				LogValue fileProperty = (LogValue)events[0].GetProperty("file");
 				Assert.AreEqual(@"Configuration\TargetRules.cs", fileProperty.Text);
 				Assert.AreEqual(@"SourceFile", fileProperty.Type.ToString());
-				Assert.AreEqual(@"Engine/Source/Programs/UnrealBuildTool/Configuration/TargetRules.cs", fileProperty.Properties![LogEventPropertyName.RelativePath].ToString());
-				Assert.AreEqual(@"//UE4/Main/Engine/Source/Programs/UnrealBuildTool/Configuration/TargetRules.cs@12345", fileProperty.Properties![LogEventPropertyName.DepotPath].ToString());
+				Assert.AreEqual(@"Engine/Source/Programs/UnrealBuildTool/Configuration/TargetRules.cs", fileProperty.Properties![LogEventPropertyName.RelativePath]?.ToString());
+				Assert.AreEqual(@"//UE4/Main/Engine/Source/Programs/UnrealBuildTool/Configuration/TargetRules.cs@12345", fileProperty.Properties![LogEventPropertyName.DepotPath]?.ToString());
 			}
 		}
 
@@ -282,8 +282,8 @@ namespace AutomationTool.Tests
 				LogValue fileProperty = (LogValue)logEvents[0].GetProperty("file");
 				Assert.AreEqual(absPath, fileProperty.Text);
 				Assert.AreEqual(@"SourceFile", fileProperty.Type.ToString());
-				Assert.AreEqual(@"Engine/Source/Runtime/CoreOnline/CoreOnline.Build.cs", fileProperty.Properties![LogEventPropertyName.RelativePath].ToString());
-				Assert.AreEqual(@"//UE4/Main/Engine/Source/Runtime/CoreOnline/CoreOnline.Build.cs@12345", fileProperty.Properties![LogEventPropertyName.DepotPath].ToString());
+				Assert.AreEqual(@"Engine/Source/Runtime/CoreOnline/CoreOnline.Build.cs", fileProperty.Properties![LogEventPropertyName.RelativePath]?.ToString());
+				Assert.AreEqual(@"//UE4/Main/Engine/Source/Runtime/CoreOnline/CoreOnline.Build.cs@12345", fileProperty.Properties![LogEventPropertyName.DepotPath]?.ToString());
 
 				CheckEventGroup(logEvents.Slice(1, 1), 1, 1, LogLevel.Warning, KnownLogEvents.Compiler);
 				Assert.AreEqual("CS0246", logEvents[1].GetProperty("code").ToString());
@@ -292,8 +292,8 @@ namespace AutomationTool.Tests
 				LogValue fileProperty1 = logEvents[1].GetProperty<LogValue>("file");
 				Assert.AreEqual(absPath, fileProperty1.Text);
 				Assert.AreEqual(@"SourceFile", fileProperty1.Type.ToString());
-				Assert.AreEqual(@"Engine/Source/Runtime/CoreOnline/CoreOnline.Build.cs", fileProperty1.Properties![LogEventPropertyName.RelativePath].ToString());
-				Assert.AreEqual(@"//UE4/Main/Engine/Source/Runtime/CoreOnline/CoreOnline.Build.cs@12345", fileProperty1.Properties![LogEventPropertyName.DepotPath].ToString());
+				Assert.AreEqual(@"Engine/Source/Runtime/CoreOnline/CoreOnline.Build.cs", fileProperty1.Properties![LogEventPropertyName.RelativePath]?.ToString());
+				Assert.AreEqual(@"//UE4/Main/Engine/Source/Runtime/CoreOnline/CoreOnline.Build.cs@12345", fileProperty1.Properties![LogEventPropertyName.DepotPath]?.ToString());
 			}
 		}
 
@@ -386,7 +386,7 @@ namespace AutomationTool.Tests
 
 				LogValue fileProperty0 = (LogValue)logEvents[0].GetProperty("file");
 				Assert.AreEqual(@"SourceFile", fileProperty0.Type.ToString());
-				Assert.AreEqual(@"Foo/Bar.txt", fileProperty0.Properties![LogEventPropertyName.RelativePath].ToString());
+				Assert.AreEqual(@"Foo/Bar.txt", fileProperty0.Properties![LogEventPropertyName.RelativePath]?.ToString());
 
 				// 1
 				CheckEventGroup(logEvents.Slice(1, 1), 1, 1, LogLevel.Warning, KnownLogEvents.Microsoft);
@@ -396,7 +396,7 @@ namespace AutomationTool.Tests
 
 				LogValue fileProperty1 = logEvents[1].GetProperty<LogValue>("file");
 				Assert.AreEqual(@"SourceFile", fileProperty1.Type.ToString());
-				Assert.AreEqual(@"Foo/Bar.txt", fileProperty1.Properties![LogEventPropertyName.RelativePath].ToString());
+				Assert.AreEqual(@"Foo/Bar.txt", fileProperty1.Properties![LogEventPropertyName.RelativePath]?.ToString());
 
 				// 2
 				CheckEventGroup(logEvents.Slice(2, 1), 2, 1, LogLevel.Error, KnownLogEvents.Microsoft);
@@ -436,7 +436,7 @@ namespace AutomationTool.Tests
 
 				LogValue noteProperty1 = logEvents[2].GetProperty<LogValue>("file");
 				Assert.AreEqual(@"SourceFile", noteProperty1.Type.ToString());
-				Assert.AreEqual(@"//UE4/Main/Engine/Plugins/Editor/EditorScriptingUtilities/Source/EditorScriptingUtilities/Public/EditorLevelLibrary.h@12345", noteProperty1.Properties![LogEventPropertyName.DepotPath].ToString());
+				Assert.AreEqual(@"//UE4/Main/Engine/Plugins/Editor/EditorScriptingUtilities/Source/EditorScriptingUtilities/Public/EditorLevelLibrary.h@12345", noteProperty1.Properties![LogEventPropertyName.DepotPath]?.ToString());
 			}
 		}
 
@@ -454,8 +454,8 @@ namespace AutomationTool.Tests
 
 			LogValue assetProperty = events[0].GetProperty<LogValue>("asset");
 			Assert.AreEqual(new Utf8String("Asset"), assetProperty.Type);
-			Assert.AreEqual(@"//UE4/Main/QAGame/Plugins/NiagaraFluids/Content/Blueprints/Phsyarum_BP.uasset@12345", assetProperty.Properties![LogEventPropertyName.DepotPath].ToString());
-			Assert.AreEqual(@"QAGame/Plugins/NiagaraFluids/Content/Blueprints/Phsyarum_BP.uasset", assetProperty.Properties![LogEventPropertyName.RelativePath].ToString());
+			Assert.AreEqual(@"//UE4/Main/QAGame/Plugins/NiagaraFluids/Content/Blueprints/Phsyarum_BP.uasset@12345", assetProperty.Properties![LogEventPropertyName.DepotPath]?.ToString());
+			Assert.AreEqual(@"QAGame/Plugins/NiagaraFluids/Content/Blueprints/Phsyarum_BP.uasset", assetProperty.Properties![LogEventPropertyName.RelativePath]?.ToString());
 		}
 
 		[TestMethod]
@@ -854,7 +854,7 @@ namespace AutomationTool.Tests
 				LogEvent logEvent = logEvents[0];
 				Assert.AreEqual("LogGatherTextFromSourceCommandlet", logEvent.GetProperty("channel").ToString());
 				Assert.AreEqual("Warning", logEvent.GetProperty("severity").ToString());
-				Assert.AreEqual("Engine/Plugins/Enterprise/VariantManager/Source/VariantManager/Private/SVariantManager.cpp", logEvent.GetProperty<LogValue>("file")!.Properties![LogEventPropertyName.RelativePath].ToString());
+				Assert.AreEqual("Engine/Plugins/Enterprise/VariantManager/Source/VariantManager/Private/SVariantManager.cpp", logEvent.GetProperty<LogValue>("file")!.Properties![LogEventPropertyName.RelativePath]?.ToString());
 				Assert.AreEqual("3717", logEvent.GetProperty("line").ToString());
 				Assert.AreEqual(0, logEvent.LineIndex);
 				Assert.AreEqual(1, logEvent.LineCount);
@@ -862,7 +862,7 @@ namespace AutomationTool.Tests
 				logEvent = logEvents[1];
 				Assert.AreEqual("LogLocTextHelper", logEvent.GetProperty("channel").ToString());
 				Assert.AreEqual("Warning", logEvent.GetProperty("severity").ToString());
-				Assert.AreEqual("Engine/Plugins/Experimental/UVEditor/Source/UVEditor/Private/UVEditorCommands.cpp", logEvent.GetProperty<LogValue>("file")!.Properties![LogEventPropertyName.RelativePath].ToString());
+				Assert.AreEqual("Engine/Plugins/Experimental/UVEditor/Source/UVEditor/Private/UVEditorCommands.cpp", logEvent.GetProperty<LogValue>("file")!.Properties![LogEventPropertyName.RelativePath]?.ToString());
 				Assert.AreEqual("39", logEvent.GetProperty("line").ToString());
 				Assert.AreEqual(0, logEvent.LineIndex);
 				Assert.AreEqual(2, logEvent.LineCount);
@@ -870,7 +870,7 @@ namespace AutomationTool.Tests
 				logEvent = logEvents[2];
 				Assert.AreEqual("LogLocTextHelper", logEvent.GetProperty("channel").ToString());
 				Assert.AreEqual("Warning", logEvent.GetProperty("severity").ToString());
-				Assert.AreEqual("Engine/Content/Localization/Engine/Engine.manifest", logEvent.GetProperty<LogValue>("file")!.Properties![LogEventPropertyName.RelativePath].ToString());
+				Assert.AreEqual("Engine/Content/Localization/Engine/Engine.manifest", logEvent.GetProperty<LogValue>("file")!.Properties![LogEventPropertyName.RelativePath]?.ToString());
 				Assert.AreEqual(1, logEvent.LineIndex);
 				Assert.AreEqual(2, logEvent.LineCount);
 			}
diff --git a/Engine/Source/Programs/Horde/HordeServer/Auditing/AuditLog.cs b/Engine/Source/Programs/Horde/HordeServer/Auditing/AuditLog.cs
--- a/Engine/Source/Programs/Horde/HordeServer/Auditing/AuditLog.cs
+++ b/Engine/Source/Programs/Horde/HordeServer/Auditing/AuditLog.cs
@@ -92,7 +92,7 @@ namespace HordeServer.Auditing
 				{
 					StringBuilder message = new StringBuilder();
 					StringBuilder format = new StringBuilder();
-					Dictionary<string, object> properties = new Dictionary<string, object>(StringComparer.Ordinal);
+					Dictionary<string, object?> properties = new Dictionary<string, object?>(StringComparer.Ordinal);
 
 					for (int idx = 0; idx < _scopes.Count; idx++)
 					{
@@ -115,7 +115,7 @@ namespace HordeServer.Auditing
 				return logEvent;
 			}
 
-			static void AppendMessage(LogEvent logEvent, int id, StringBuilder message, StringBuilder format, Dictionary<string, object> properties)
+			static void AppendMessage(LogEvent logEvent, int id, StringBuilder message, StringBuilder format, Dictionary<string, object?> properties)
 			{
 				message.Append(logEvent.Message);
 				if (logEvent.Format == null)
@@ -128,7 +128,7 @@ namespace HordeServer.Auditing
 					format.Append(logEvent.Format);
 					if (logEvent.Properties != null)
 					{
-						foreach (KeyValuePair<string, object> item in logEvent.Properties)
+						foreach (KeyValuePair<string, object?> item in logEvent.Properties)
 						{
 							properties[item.Key] = item.Value;
 						}
diff --git a/Engine/Source/Programs/Horde/Plugins/Build/HordeServer.Build.Tests/Issues/IssueServiceTests.cs b/Engine/Source/Programs/Horde/Plugins/Build/HordeServer.Build.Tests/Issues/IssueServiceTests.cs
--- a/Engine/Source/Programs/Horde/Plugins/Build/HordeServer.Build.Tests/Issues/IssueServiceTests.cs
+++ b/Engine/Source/Programs/Horde/Plugins/Build/HordeServer.Build.Tests/Issues/IssueServiceTests.cs
@@ -1956,7 +1956,7 @@ namespace HordeServer.Tests.Issues
 			Assert.AreEqual("Errors in Update Version Files", issue.Summary);
 		}
 
-		static private IEnumerable<JsonLogEvent> MultilineLogEvent(LogLevel level, EventId eventId, string format, Dictionary<string, object> properties)
+		static private IEnumerable<JsonLogEvent> MultilineLogEvent(LogLevel level, EventId eventId, string format, Dictionary<string, object?> properties)
 		{
 			DateTime time = new DateTime(2024, 1, 1, 0, 0, 0, 0, DateTimeKind.Utc);
 			string message = MessageTemplate.Render(format, properties!);
@@ -2035,7 +2035,7 @@ namespace HordeServer.Tests.Issues
 				await UpdateCompleteStepAsync(job, 0, 0, JobStepOutcome.Failure);
 				await using (TestJsonLogger logger = await CreateLoggerAsync(job, 0, 1))
 				{
-					Dictionary<string, object> properties = new Dictionary<string, object>
+					Dictionary<string, object?> properties = new Dictionary<string, object?>
 					{
 						["Summary"] = logMessage,
 						["Callstack"] = logAlternateCallstack
diff --git a/Engine/Source/Programs/Shared/EpicGames.Core.Tests/LogEventTests.cs b/Engine/Source/Programs/Shared/EpicGames.Core.Tests/LogEventTests.cs
new file mode 100644
--- /dev/null
+++ b/Engine/Source/Programs/Shared/EpicGames.Core.Tests/LogEventTests.cs
@@ -0,0 +1,52 @@
+// Copyright Epic Games, Inc.All Rights Reserved.
+
+using System;
+using System.Collections.Generic;
+using Microsoft.Extensions.Logging;
+using Microsoft.VisualStudio.TestTools.UnitTesting;
+
+namespace EpicGames.Core.Tests
+{
+	[TestClass]
+	public class LogEventTests
+	{
+		class TestClass
+		{
+			public int Foo { get; set; }
+			public string? Bar { get; set; }
+
+			public override string ToString()
+				=> $"TestClass({Foo},{Bar})";
+		}
+
+		[TestMethod]
+		public void ArgumentTests()
+		{
+			CaptureLogger logger = new CaptureLogger();
+			logger.LogInformation("Test {Value} {@Value}", new TestClass { Foo = 123, Bar = "456" }, new TestClass { Foo = 789, Bar = "hello" });
+
+			LogEvent logEvent = logger.Events[0];
+			logEvent.Time = new DateTime(2024, 1, 1);
+
+			Assert.AreEqual("Test TestClass(123,456) {\"foo\":789,\"bar\":\"hello\"}", logEvent.ToString());
+		}
+
+		[TestMethod]
+		public void RoundTripTests()
+		{
+			DateTime time = new DateTime(2024, 1, 1);
+			LogEvent ev = new LogEvent(time, LogLevel.Information, default, "hello", null, new KeyValuePair<string, object?>[] { KeyValuePair.Create<string, object?>("foo", new LogValue(new Utf8String("type"), "text")) }, null);
+
+			JsonLogEvent jev = JsonLogEvent.Parse(ev.ToJsonBytes());
+			string jevs = jev.ToString();
+			Assert.AreEqual("{\"time\":\"2024-01-01T00:00:00\",\"level\":\"Information\",\"message\":\"hello\",\"properties\":{\"foo\":{\"$type\":\"type\",\"$text\":\"text\"}}}", jevs);
+
+			LogEvent ev2 = LogEvent.Read(jev.Data.Span);
+			Assert.AreEqual(ev.ToString(), ev2.ToString());
+
+			JsonLogEvent jev2 = JsonLogEvent.Parse(ev2.ToJsonBytes());
+			string jevs2 = jev2.ToString();
+			Assert.AreEqual(jevs, jevs2);
+		}
+	}
+}
diff --git a/Engine/Source/Programs/Shared/EpicGames.Core/Log.cs b/Engine/Source/Programs/Shared/EpicGames.Core/Log.cs
--- a/Engine/Source/Programs/Shared/EpicGames.Core/Log.cs
+++ b/Engine/Source/Programs/Shared/EpicGames.Core/Log.cs
@@ -791,14 +791,14 @@ namespace EpicGames.Core
 		class Scope : IDisposable
 		{
 			readonly LoggerScopeCollection _outer;
-			readonly IEnumerable<KeyValuePair<string, object>>? _properties;
+			readonly IEnumerable<KeyValuePair<string, object?>>? _properties;
 
-			public IEnumerable<KeyValuePair<string, object>>? Properties => _properties;
+			public IEnumerable<KeyValuePair<string, object?>>? Properties => _properties;
 
 			public Scope(LoggerScopeCollection outer, object? value)
 			{
 				_outer = outer;
-				_properties = value as IEnumerable<KeyValuePair<string, object>>;
+				_properties = value as IEnumerable<KeyValuePair<string, object?>>;
 
 				_outer._scopes.Add(this);
 			}
@@ -824,13 +824,13 @@ namespace EpicGames.Core
 		/// <summary>
 		/// Enumerates all current properties
 		/// </summary>
-		public IEnumerable<KeyValuePair<string, object>> GetProperties()
+		public IEnumerable<KeyValuePair<string, object?>> GetProperties()
 		{
 			foreach (Scope scope in _scopes)
 			{
 				if (scope.Properties != null)
 				{
-					foreach (KeyValuePair<string, object> property in scope.Properties)
+					foreach (KeyValuePair<string, object?> property in scope.Properties)
 					{
 						yield return property;
 					}
diff --git a/Engine/Source/Programs/Shared/EpicGames.Core/LogEvent.cs b/Engine/Source/Programs/Shared/EpicGames.Core/LogEvent.cs
--- a/Engine/Source/Programs/Shared/EpicGames.Core/LogEvent.cs
+++ b/Engine/Source/Programs/Shared/EpicGames.Core/LogEvent.cs
@@ -89,21 +89,21 @@ namespace EpicGames.Core
 		/// <summary>
 		/// Map of property name to value
 		/// </summary>
-		public IEnumerable<KeyValuePair<string, object>>? Properties { get; set; }
+		public IEnumerable<KeyValuePair<string, object?>>? Properties { get; set; }
 
 		/// <summary>
 		/// The exception value
 		/// </summary>
 		public LogException? Exception { get; }
 
-		class MergedPropertyList : IEnumerable<KeyValuePair<string, object>>
+		class MergedPropertyList : IEnumerable<KeyValuePair<string, object?>>
 		{
 			readonly HashSet<string> _names = new HashSet<string>();
-			readonly List<KeyValuePair<string, object>> _properties = new List<KeyValuePair<string, object>>();
+			readonly List<KeyValuePair<string, object?>> _properties = new List<KeyValuePair<string, object?>>();
 
-			public void AddRange(IEnumerable<KeyValuePair<string, object>> properties)
+			public void AddRange(IEnumerable<KeyValuePair<string, object?>> properties)
 			{
-				foreach (KeyValuePair<string, object> property in properties)
+				foreach (KeyValuePair<string, object?> property in properties)
 				{
 					if (_names.Add(property.Key))
 					{
@@ -113,7 +113,7 @@ namespace EpicGames.Core
 			}
 
 			/// <inheritdoc/>
-			public IEnumerator<KeyValuePair<string, object>> GetEnumerator()
+			public IEnumerator<KeyValuePair<string, object?>> GetEnumerator()
 				=> _properties.GetEnumerator();
 
 			/// <inheritdoc/>
@@ -126,7 +126,7 @@ namespace EpicGames.Core
 		/// <summary>
 		/// Constructor
 		/// </summary>
-		public LogEvent(DateTime time, LogLevel level, EventId eventId, string message, string? format, IEnumerable<KeyValuePair<string, object>>? properties, LogException? exception)
+		public LogEvent(DateTime time, LogLevel level, EventId eventId, string message, string? format, IEnumerable<KeyValuePair<string, object?>>? properties, LogException? exception)
 			: this(time, level, eventId, 0, 1, message, format, properties, exception)
 		{
 		}
@@ -134,7 +134,7 @@ namespace EpicGames.Core
 		/// <summary>
 		/// Constructor
 		/// </summary>
-		public LogEvent(DateTime time, LogLevel level, EventId eventId, int lineIndex, int lineCount, string message, string? format, IEnumerable<KeyValuePair<string, object>>? properties, LogException? exception)
+		public LogEvent(DateTime time, LogLevel level, EventId eventId, int lineIndex, int lineCount, string message, string? format, IEnumerable<KeyValuePair<string, object?>>? properties, LogException? exception)
 		{
 			Time = time;
 			Level = level;
@@ -152,14 +152,14 @@ namespace EpicGames.Core
 		/// </summary>
 		/// <param name="name">Name of the property to add</param>
 		/// <param name="value">Value for the property</param>
-		public void AddProperty(string name, object value)
+		public void AddProperty(string name, object? value)
 			=> AddProperties(new[] { KeyValuePair.Create(name, value) });
 
 		/// <summary>
 		/// Add new properties to this event
 		/// </summary>
 		/// <param name="properties">Properties to add</param>
-		public void AddProperties(IEnumerable<KeyValuePair<string, object>> properties)
+		public void AddProperties(IEnumerable<KeyValuePair<string, object?>> properties)
 		{
 			if (Properties == null)
 			{
@@ -211,9 +211,9 @@ namespace EpicGames.Core
 		{
 			if (Properties != null)
 			{
-				foreach (KeyValuePair<string, object> pair in Properties)
+				foreach (KeyValuePair<string, object?> pair in Properties)
 				{
-					if (pair.Key.Equals(name, StringComparison.Ordinal))
+					if (pair.Key.Equals(name, StringComparison.Ordinal) && pair.Value != null)
 					{
 						value = pair.Value;
 						return true;
@@ -235,7 +235,7 @@ namespace EpicGames.Core
 		public bool TryGetProperty<T>(string name, [NotNullWhen(true)] out T value)
 		{
 			object? untypedValue;
-			if(TryGetProperty(name, out untypedValue) && untypedValue is T typedValue)
+			if (TryGetProperty(name, out untypedValue) && untypedValue is T typedValue)
 			{
 				value = typedValue;
 				return true;
@@ -274,8 +274,8 @@ namespace EpicGames.Core
 			int line = 0;
 			int lineCount = 1;
 			string message = String.Empty;
-			string format = String.Empty;
-			Dictionary<string, object>? properties = null;
+			string? format = null;
+			Dictionary<string, object?>? properties = null;
 			LogException? exception = null;
 
 			ReadOnlySpan<byte> propertyName;
@@ -318,33 +318,36 @@ namespace EpicGames.Core
 			return new LogEvent(time, level, eventId, line, lineCount, message, format, properties, exception);
 		}
 
-		static Dictionary<string, object> ReadProperties(ref Utf8JsonReader reader)
+		static Dictionary<string, object?> ReadProperties(ref Utf8JsonReader reader)
 		{
-			Dictionary<string, object> properties = new Dictionary<string, object>();
+			Dictionary<string, object?> properties = new Dictionary<string, object?>();
 
 			ReadOnlySpan<byte> propertyName;
 			for (; JsonExtensions.TryReadNextPropertyName(ref reader, out propertyName); reader.Skip())
 			{
 				string name = Encoding.UTF8.GetString(propertyName);
-				object value = ReadPropertyValue(ref reader);
+
+				object? value = ReadPropertyValue(ref reader);
 				properties.Add(name, value);
 			}
 
 			return properties;
 		}
 
-		static object ReadPropertyValue(ref Utf8JsonReader reader)
+		static object? ReadPropertyValue(ref Utf8JsonReader reader)
 		{
 			switch (reader.TokenType)
 			{
 				case JsonTokenType.Null:
-					return null!;
+					return null;
 				case JsonTokenType.True:
 					return true;
 				case JsonTokenType.False:
 					return false;
+				case JsonTokenType.StartArray:
+					return ReadArrayPropertyValue(ref reader);
 				case JsonTokenType.StartObject:
-					return ReadStructuredPropertyValue(ref reader);
+					return ReadObjectPropertyValue(ref reader);
 				case JsonTokenType.String:
 					return reader.GetString()!;
 				case JsonTokenType.Number:
@@ -365,31 +368,42 @@ namespace EpicGames.Core
 			}
 		}
 
-		static LogValue ReadStructuredPropertyValue(ref Utf8JsonReader reader)
+		static object ReadArrayPropertyValue(ref Utf8JsonReader reader)
 		{
-			string type = String.Empty;
-			string text = String.Empty;
-			Dictionary<Utf8String, object>? properties = null;
+			List<object?> result = new List<object?>();
+			while (reader.Read() && reader.TokenType != JsonTokenType.EndArray)
+			{
+				result.Add(ReadPropertyValue(ref reader));
+			}
+			return result;
+		}
 
-			ReadOnlySpan<byte> propertyName;
-			for (; JsonExtensions.TryReadNextPropertyName(ref reader, out propertyName); reader.Skip())
+		static object ReadObjectPropertyValue(ref Utf8JsonReader reader)
+		{
+			// Read all the properties
+			Dictionary<Utf8String, object?>? properties = new Dictionary<Utf8String, object?>();
+			for (; JsonExtensions.TryReadNextPropertyName(ref reader, out ReadOnlySpan<byte> propertyName); reader.Skip())
 			{
-				if (Utf8StringComparer.OrdinalIgnoreCase.Equals(propertyName, LogEventPropertyName.Type.Span))
-				{
-					type = reader.GetString() ?? String.Empty;
-				}
-				else if (Utf8StringComparer.OrdinalIgnoreCase.Equals(propertyName, LogEventPropertyName.Text.Span))
+				properties.Add(new Utf8String(propertyName.ToArray()), ReadPropertyValue(ref reader));
+			}
+
+			// Check if we can convert it to a LogValue
+			if (properties.TryGetValue(LogEventPropertyName.Type, out object? type) && type is string typeStr && properties.TryGetValue(LogEventPropertyName.Text, out object? text) && text is string textStr)
+			{
+				properties.Remove(LogEventPropertyName.Type);
+				properties.Remove(LogEventPropertyName.Text);
+
+				if (properties.Count == 0)
 				{
-					text = reader.GetString() ?? String.Empty;
+					return new LogValue(new Utf8String(typeStr), textStr);
 				}
 				else
 				{
-					properties ??= new Dictionary<Utf8String, object>();
-					properties.Add(new Utf8String(propertyName.ToArray()), ReadPropertyValue(ref reader));
+					return new LogValue(new Utf8String(typeStr), textStr, properties);
 				}
 			}
 
-			return new LogValue(new Utf8String(type), text, properties);
+			return properties;
 		}
 
 		/// <summary>
@@ -433,7 +447,7 @@ namespace EpicGames.Core
 						if (name.StartsWith("@", StringComparison.Ordinal))
 						{
 							writer.WritePropertyName(name[1..]);
-							JsonSerializer.Serialize(writer, value, value.GetType(), s_jsonSerializerOptions);
+							JsonSerializer.Serialize(writer, value, value?.GetType() ?? typeof(object), s_jsonSerializerOptions);
 						}
 						else
 						{
@@ -502,7 +516,7 @@ namespace EpicGames.Core
 		/// </summary>
 		public static LogEvent Create(LogLevel level, EventId eventId, Exception? exception, string format, params object[] args)
 		{
-			Dictionary<string, object> properties = new Dictionary<string, object>();
+			Dictionary<string, object?> properties = new Dictionary<string, object?>();
 			MessageTemplate.ParsePropertyValues(format, args, properties);
 
 			string message = MessageTemplate.Render(format, properties!);
@@ -514,6 +528,8 @@ namespace EpicGames.Core
 		/// </summary>
 		public static LogEvent FromState<TState>(LogLevel level, EventId eventId, TState state, Exception? exception, Func<TState, Exception?, string> formatter)
 		{
+			_ = formatter;
+
 			if (state is LogEvent logEvent)
 			{
 				return logEvent;
@@ -525,12 +541,10 @@ namespace EpicGames.Core
 
 			DateTime time = DateTime.UtcNow;
 
-			// Render the message
-			string message = formatter(state, exception);
-
 			// Try to log the event
-			IEnumerable<KeyValuePair<string, object>>? values = state as IEnumerable<KeyValuePair<string, object>>;
+			IEnumerable<KeyValuePair<string, object?>>? values = state as IEnumerable<KeyValuePair<string, object?>>;
 			string? format = values?.FirstOrDefault(x => x.Key.Equals(MessageTemplate.FormatPropertyName, StringComparison.Ordinal)).Value?.ToString();
+			string message = MessageTemplate.Render(format ?? String.Empty, values);
 			return new LogEvent(time, level, eventId, message, format, values, LogException.FromException(exception));
 		}
 
@@ -687,11 +701,11 @@ namespace EpicGames.Core
 		/// Log events received
 		/// </summary>
 		public IReadOnlyList<LogEvent> LogEvents => _logEvents;
-		
+
 		private readonly List<LogEvent> _logEvents = new List<LogEvent>();
 		private readonly int[] _includeEventIds;
 		private readonly Action<LogEvent>? _eventCallback;
-		
+
 		/// <summary>
 		/// Constructor
 		/// </summary>
@@ -713,7 +727,7 @@ namespace EpicGames.Core
 			}
 		}
 	}
-	
+
 	/// <summary>
 	/// Converter for serialization of <see cref="LogEvent"/> instances to Json streams
 	/// </summary>
diff --git a/Engine/Source/Programs/Shared/EpicGames.Core/LogEventBuilder.cs b/Engine/Source/Programs/Shared/EpicGames.Core/LogEventBuilder.cs
--- a/Engine/Source/Programs/Shared/EpicGames.Core/LogEventBuilder.cs
+++ b/Engine/Source/Programs/Shared/EpicGames.Core/LogEventBuilder.cs
@@ -438,10 +438,10 @@ namespace EpicGames.Core
 
 		static LogEvent CreateEvent(DateTime time, LogLevel level, EventId eventId, int lineIndex, int lineCount, LogLine line)
 		{
-			Dictionary<string, object>? properties = null;
+			Dictionary<string, object?>? properties = null;
 			if (line._properties != null)
 			{
-				properties = new Dictionary<string, object>();
+				properties = new Dictionary<string, object?>();
 				foreach ((string name, object value) in line._properties)
 				{
 					object newValue;
diff --git a/Engine/Source/Programs/Shared/EpicGames.Core/LogEventMarkup.cs b/Engine/Source/Programs/Shared/EpicGames.Core/LogEventMarkup.cs
--- a/Engine/Source/Programs/Shared/EpicGames.Core/LogEventMarkup.cs
+++ b/Engine/Source/Programs/Shared/EpicGames.Core/LogEventMarkup.cs
@@ -60,7 +60,7 @@ namespace EpicGames.Core
 		/// </summary>
 		public static void AnnotateSourceFile(this LogEventBuilder builder, Group group, string? baseDir)
 		{
-			Dictionary<Utf8String, object>? properties = null;
+			Dictionary<Utf8String, object?>? properties = null;
 			if (!String.IsNullOrEmpty(baseDir))
 			{
 				string file = group.Value;
@@ -69,7 +69,7 @@ namespace EpicGames.Core
 					try
 					{
 						string combinedPath = RemoveRelativeDirs(Path.Combine(baseDir, file));
-						properties = new Dictionary<Utf8String, object>();
+						properties = new Dictionary<Utf8String, object?>();
 						properties[LogEventPropertyName.File] = combinedPath;
 					}
 					catch
@@ -105,7 +105,7 @@ namespace EpicGames.Core
 			identifier = Regex.Replace(identifier, "^.* ", "");
 
 			// Add it to the list
-			Dictionary<Utf8String, object> properties = new Dictionary<Utf8String, object>();
+			Dictionary<Utf8String, object?> properties = new Dictionary<Utf8String, object?>();
 			properties[LogEventPropertyName.Identifier] = identifier;
 			builder.Annotate(group, new LogValue(LogValueType.Symbol, "", properties));
 		}
diff --git a/Engine/Source/Programs/Shared/EpicGames.Core/LogValue.cs b/Engine/Source/Programs/Shared/EpicGames.Core/LogValue.cs
--- a/Engine/Source/Programs/Shared/EpicGames.Core/LogValue.cs
+++ b/Engine/Source/Programs/Shared/EpicGames.Core/LogValue.cs
@@ -65,7 +65,7 @@ namespace EpicGames.Core
 		/// <summary>
 		/// Properties associated with the value
 		/// </summary>
-		public Dictionary<Utf8String, object>? Properties { get; }
+		public Dictionary<Utf8String, object?>? Properties { get; }
 
 		/// <summary>
 		/// Constructor
@@ -73,7 +73,7 @@ namespace EpicGames.Core
 		/// <param name="type">Type of the value</param>
 		/// <param name="text">Rendering of the value as text</param>
 		/// <param name="properties">Additional properties for this value</param>
-		public LogValue(Utf8String type, string text, Dictionary<Utf8String, object>? properties = null)
+		public LogValue(Utf8String type, string text, Dictionary<Utf8String, object?>? properties = null)
 		{
 			Type = type;
 			Text = text;
@@ -87,7 +87,7 @@ namespace EpicGames.Core
 		/// <param name="text">Display text for the fiel</param>
 		public static LogValue SourceFile(FileReference file, string text)
 		{
-			return new LogValue(LogValueType.SourceFile, text, new Dictionary<Utf8String, object> { [LogEventPropertyName.File] = file.FullName });
+			return new LogValue(LogValueType.SourceFile, text, new Dictionary<Utf8String, object?> { [LogEventPropertyName.File] = file.FullName });
 		}
 
 		/// <summary>
@@ -106,7 +106,7 @@ namespace EpicGames.Core
 		/// <param name="text">Text to render for the link</param>
 		public static LogValue Link(Uri target, string text)
 		{
-			return new LogValue(LogValueType.Link, text, new Dictionary<Utf8String, object> { [LogEventPropertyName.Target] = target.ToString() });
+			return new LogValue(LogValueType.Link, text, new Dictionary<Utf8String, object?> { [LogEventPropertyName.Target] = target.ToString() });
 		}
 
 		/// <summary>
@@ -127,13 +127,13 @@ namespace EpicGames.Core
 		{
 			Type objType = obj.GetType();
 
-			Dictionary<Utf8String, object>? properties = null;
+			Dictionary<Utf8String, object?>? properties = null;
 			foreach (PropertyInfo propertyInfo in objType.GetProperties())
 			{
 				if (propertyInfo.GetCustomAttribute<JsonIgnoreAttribute>() == null)
 				{
 					string name = propertyInfo.GetCustomAttribute<JsonPropertyNameAttribute>()?.Name ?? propertyInfo.Name;
-					properties ??= new Dictionary<Utf8String, object>();
+					properties ??= new Dictionary<Utf8String, object?>();
 					properties[new Utf8String(name)] = propertyInfo.GetValue(obj)!;
 				}
 			}
diff --git a/Engine/Source/Programs/Shared/EpicGames.Core/LogValueFormatter.cs b/Engine/Source/Programs/Shared/EpicGames.Core/LogValueFormatter.cs
--- a/Engine/Source/Programs/Shared/EpicGames.Core/LogValueFormatter.cs
+++ b/Engine/Source/Programs/Shared/EpicGames.Core/LogValueFormatter.cs
@@ -149,7 +149,7 @@ namespace EpicGames.Core
 				writer.WriteString(LogEventPropertyName.Text, logValue.Text);
 				if (logValue.Properties != null)
 				{
-					foreach (KeyValuePair<Utf8String, object> pair in logValue.Properties)
+					foreach (KeyValuePair<Utf8String, object?> pair in logValue.Properties)
 					{
 						writer.WritePropertyName(pair.Key);
 						WritePropertyValue(pair.Value, writer);
@@ -231,6 +231,36 @@ namespace EpicGames.Core
 			}
 		}
 
+		class DictionaryFormatter : ILogValueFormatter
+		{
+			public void Format(object value, Utf8JsonWriter writer)
+			{
+				Dictionary<string, object> dictionary = (Dictionary<string, object>)value;
+				writer.WriteStartObject();
+				foreach (KeyValuePair<string, object> kvp in dictionary)
+				{
+					writer.WritePropertyName(kvp.Key);
+					Format(kvp.Value, writer);
+				}
+				writer.WriteEndObject();
+			}
+		}
+
+		class Utf8DictionaryFormatter : ILogValueFormatter
+		{
+			public void Format(object value, Utf8JsonWriter writer)
+			{
+				Dictionary<Utf8String, object> dictionary = (Dictionary<Utf8String, object>)value;
+				writer.WriteStartObject();
+				foreach (KeyValuePair<Utf8String, object> kvp in dictionary)
+				{
+					writer.WritePropertyName(kvp.Key);
+					LogValueFormatter.Format(kvp.Value, writer);
+				}
+				writer.WriteEndObject();
+			}
+		}
+
 		class FileReferenceFormatter : ILogValueFormatter
 		{
 			public void Format(object value, Utf8JsonWriter writer)
@@ -258,6 +288,8 @@ namespace EpicGames.Core
 			formatters.TryAdd(typeof(float), new FloatFormatter());
 			formatters.TryAdd(typeof(double), new DoubleFormatter());
 			formatters.TryAdd(typeof(string), s_stringFormatter);
+			formatters.TryAdd(typeof(Dictionary<string, object>), new DictionaryFormatter());
+			formatters.TryAdd(typeof(Dictionary<Utf8String, object>), new Utf8DictionaryFormatter());
 			formatters.TryAdd(typeof(LogValue), new StructuredLogValueFormatter());
 			formatters.TryAdd(typeof(FileReference), new FileReferenceFormatter());
 			formatters.TryAdd(typeof(Activity), new ActivityFormatter());
diff --git a/Engine/Source/Programs/Shared/EpicGames.Core/MessageTemplate.cs b/Engine/Source/Programs/Shared/EpicGames.Core/MessageTemplate.cs
--- a/Engine/Source/Programs/Shared/EpicGames.Core/MessageTemplate.cs
+++ b/Engine/Source/Programs/Shared/EpicGames.Core/MessageTemplate.cs
@@ -175,7 +175,7 @@ namespace EpicGames.Core
 		/// <param name="args">Argument list to parse</param>
 		/// <param name="properties"></param>
 		/// <returns></returns>
-		public static void ParsePropertyValues(string format, object[] args, Dictionary<string, object> properties)
+		public static void ParsePropertyValues(string format, object[] args, Dictionary<string, object?> properties)
 		{
 			List<(int, int)>? offsets = ParsePropertyNames(format);
 			if (offsets != null)
diff --git a/Engine/Source/Programs/Shared/EpicGames.Horde/Logs/ServerLogPacketBuilder.cs b/Engine/Source/Programs/Shared/EpicGames.Horde/Logs/ServerLogPacketBuilder.cs
--- a/Engine/Source/Programs/Shared/EpicGames.Horde/Logs/ServerLogPacketBuilder.cs
+++ b/Engine/Source/Programs/Shared/EpicGames.Horde/Logs/ServerLogPacketBuilder.cs
@@ -127,7 +127,7 @@ namespace EpicGames.Horde.Logs
 				}
 
 				string escapedLine = escapedLineBuilder.ToString();
-				KeyValuePair<string, object>[] properties = new[] { new KeyValuePair<string, object>("Text", escapedLine) };
+				KeyValuePair<string, object?>[] properties = new[] { new KeyValuePair<string, object?>("Text", escapedLine) };
 				LogEvent newLogEvent = new LogEvent(DateTime.UtcNow, LogLevel.Error, default, $"Invalid json log event: {escapedLineBuilder}", "Invalid json log event: {Text}", properties, LogException.FromException(ex));
 				JsonLogEvent newJsonLogEvent = new JsonLogEvent(newLogEvent);
 
diff --git a/Engine/Source/Programs/UnrealBuildTool.Tests/LogEventMatcherTests.cs b/Engine/Source/Programs/UnrealBuildTool.Tests/LogEventMatcherTests.cs
--- a/Engine/Source/Programs/UnrealBuildTool.Tests/LogEventMatcherTests.cs
+++ b/Engine/Source/Programs/UnrealBuildTool.Tests/LogEventMatcherTests.cs
@@ -41,7 +41,7 @@ namespace UnrealBuildToolTests
 				LogEvent logEvent = LogEvent.Read(JsonLogEvent.FromLoggerState(logLevel, eventId, state, exception, formatter).Data.Span);
 				if (logEvent.Level != LogLevel.Information || logEvent.Id != default || logEvent.Properties != null)
 				{
-					KeyValuePair<string, object>[] items = new[] { new KeyValuePair<string, object>(LogLine, _logLineIndex) };
+					KeyValuePair<string, object?>[] items = new[] { new KeyValuePair<string, object?>(LogLine, _logLineIndex) };
 					logEvent.Properties = (logEvent.Properties == null) ? items : Enumerable.Concat(logEvent.Properties, items);
 					_events.Add(logEvent);
 				}
