Missed a file from the logger nuget update

#rnx

diff --git a/Engine/Source/Programs/AutomationTool.Tests/LogParserTests.cs b/Engine/Source/Programs/AutomationTool.Tests/LogParserTests.cs
--- a/Engine/Source/Programs/AutomationTool.Tests/LogParserTests.cs
+++ b/Engine/Source/Programs/AutomationTool.Tests/LogParserTests.cs
@@ -27,7 +27,7 @@ namespace AutomationTool.Tests
 
 			public List<LogEvent> _events = new List<LogEvent>();
 
-			public IDisposable BeginScope<TState>(TState state) => null!;
+			public IDisposable? BeginScope<TState>(TState state) where TState : notnull => null!;
 
 			public bool IsEnabled(LogLevel logLevel) => true;
 
