Add a log event matcher for running bash scripts. Improves errors in Horde server build jobs.

#rnx

diff --git a/Engine/Source/Programs/AutomationTool.Tests/LogParserTests.cs b/Engine/Source/Programs/AutomationTool.Tests/LogParserTests.cs
--- a/Engine/Source/Programs/AutomationTool.Tests/LogParserTests.cs
+++ b/Engine/Source/Programs/AutomationTool.Tests/LogParserTests.cs
@@ -124,6 +124,20 @@ namespace AutomationTool.Tests
 			CheckEventGroup(logEvents, 1, 1, LogLevel.Error, KnownLogEvents.ExitCode);
 		}
 
+		[TestMethod]
+		public void ExitCodeEventMatcher4()
+		{
+			string[] lines =
+			{
+				@"/app/Source/Programs/Horde/Plugins/Build/HordeServer.Build.Tests/Agents/AgentServiceTest.cs(221,126): error CS1503: Argument 4: cannot convert from 'EpicGames.Horde.Agents.AgentStatus?' to 'EpicGames.Horde.Agents.AgentStatus' [/app/Source/Programs/Horde/Plugins/Build/HordeServer.Build.Tests/HordeServer.Build.Tests.csproj]",
+				@"ERROR: process ""/bin/sh -c bash Source/Programs/Horde/Scripts/test.sh"" did not complete successfully: exit code: 1"
+			};
+
+			List<LogEvent> logEvents = Parse(lines);
+			CheckEventGroup(logEvents.Slice(0, 1), 0, 1, LogLevel.Error, KnownLogEvents.Compiler);
+			CheckEventGroup(logEvents.Slice(1, 1), 1, 1, LogLevel.Error, KnownLogEvents.ExitCode);
+		}
+
 		[TestMethod]
 		public void CrashEventMatcher()
 		{
diff --git a/Engine/Source/Programs/AutomationTool/AutomationUtils/Matchers/ExitCodeEventMatcher.cs b/Engine/Source/Programs/AutomationTool/AutomationUtils/Matchers/ExitCodeEventMatcher.cs
--- a/Engine/Source/Programs/AutomationTool/AutomationUtils/Matchers/ExitCodeEventMatcher.cs
+++ b/Engine/Source/Programs/AutomationTool/AutomationUtils/Matchers/ExitCodeEventMatcher.cs
@@ -16,7 +16,8 @@ namespace AutomationUtils.Matchers
 		static readonly Regex s_pattern = new Regex(
 			@"Editor terminated with exit code -?[1-9]" +
 			@"|(?:Error executing.+)(?:tool returned code)(.+)" +
-			@"|\s*RunU[AB]T ERROR:");
+			@"|\s*RunU[AB]T ERROR:" +
+			@"|\s*ERROR:.*did not complete successfully: exit code: [0-9]+");
 
 		public LogEventMatch? Match(ILogCursor cursor)
 		{
