Improve AIBotSoak ASAN logging

diff --git a/Engine/Source/Programs/AutomationTool.Tests/LogParserTests.cs b/Engine/Source/Programs/AutomationTool.Tests/LogParserTests.cs
--- a/Engine/Source/Programs/AutomationTool.Tests/LogParserTests.cs
+++ b/Engine/Source/Programs/AutomationTool.Tests/LogParserTests.cs
@@ -893,6 +893,206 @@ namespace AutomationTool.Tests
 			Assert.AreEqual("heap-use-after-free", logEvents[12].GetProperty("SummaryReason").ToString());
 		}
 
+		[TestMethod]
+		public void AddressSanitizerErrorMatcher2()
+		{
+			string[] lines =
+			{
+				@"=================================================================",
+				@"==85[0x18ca8]==ERROR: AddressSanitizer: use-after-poison on address 0x00d0572d5090 at pc 0x0000268d15d3 bp 0x0007effc1b90 sp 0x0007effc1b88",
+				@"READ of size 2 at 0x00d0572d5090 thread T0",
+				@"    #0 0x0000268d15d3  (eboot.bin+0x4c55d3)",
+				@"    #1 0x0000268cb2ee  (eboot.bin+0x4bf2ee)",
+				@"    #2 0x0000268d06b9  (eboot.bin+0x4c46b9)",
+				@"    #3 0x0000268d072a  (eboot.bin+0x4c472a)",
+				@"    #4 0x000026904cf7  (eboot.bin+0x4f8cf7)",
+				@"    #5 0x0000268d2fd1  (eboot.bin+0x4c6fd1)",
+				@"    #6 0x0000268d486c  (eboot.bin+0x4c886c)",
+				@"    #7 0x0000268caf3e  (eboot.bin+0x4bef3e)",
+				@"    #8 0x00002642a55f  (eboot.bin+0x1e55f)",
+				@"    #9 0x00002642a2cd  (eboot.bin+0x1e2cd)",
+				@"    #10 0x0000268ca9d9  (eboot.bin+0x4be9d9)",
+				@"    #11 0x0000267c663d  (eboot.bin+0x3ba63d)",
+				@"    #12 0x0000267433e9  (eboot.bin+0x3373e9)",
+				@"    #13 0x00004005f913  (eboot.bin+0x19c53913)",
+				@"    #14 0x00004005f7a1  (eboot.bin+0x19c537a1)",
+				@"    #15 0x000036b55372  (eboot.bin+0x10749372)",
+				@"    #16 0x000036bbc84e  (eboot.bin+0x107b084e)",
+				@"    #17 0x00002a5e38ec  (eboot.bin+0x41d78ec)",
+				@"    #18 0x000036b2bd70  (eboot.bin+0x1071fd70)",
+				@"    #19 0x000036b2b1d9  (eboot.bin+0x1071f1d9)",
+				@"    #20 0x000036bcd8e6  (eboot.bin+0x107c18e6)",
+				@"    #21 0x000036bcd55b  (eboot.bin+0x107c155b)",
+				@"    #22 0x000036b315a7  (eboot.bin+0x107255a7)",
+				@"    #23 0x000036b32697  (eboot.bin+0x10726697)",
+				@"    #24 0x000036b9dfde  (eboot.bin+0x10791fde)",
+				@"    #25 0x000036bd751c  (eboot.bin+0x107cb51c)",
+				@"    #26 0x000036bd72cb  (eboot.bin+0x107cb2cb)",
+				@"    #27 0x0000265cc963  (eboot.bin+0x1c0963)",
+				@"    #28 0x000031b40ed6  (eboot.bin+0xb734ed6)",
+				@"    #29 0x000031ac8d2e  (eboot.bin+0xb6bcd2e)",
+				@"    #30 0x000031b7c352  (eboot.bin+0xb770352)",
+				@"    #31 0x000031b7bfa0  (eboot.bin+0xb76ffa0)",
+				@"    #32 0x000029fdd12f  (eboot.bin+0x3bd112f)",
+				@"    #33 0x000031a5885e  (eboot.bin+0xb64c85e)",
+				@"    #34 0x00002a014420  (eboot.bin+0x3c08420)",
+				@"    #35 0x00002a0140d5  (eboot.bin+0x3c080d5)",
+				@"    #36 0x000029fdd12f  (eboot.bin+0x3bd112f)",
+				@"    #37 0x000029fed56c  (eboot.bin+0x3be156c)",
+				@"    #38 0x00002a014420  (eboot.bin+0x3c08420)",
+				@"    #39 0x00002a0140d5  (eboot.bin+0x3c080d5)",
+				@"    #40 0x000029fdd12f  (eboot.bin+0x3bd112f)",
+				@"    #41 0x000029fc93f0  (eboot.bin+0x3bbd3f0)",
+				@"    #42 0x000029fc902c  (eboot.bin+0x3bbd02c)",
+				@"    #43 0x000029fe3f84  (eboot.bin+0x3bd7f84)",
+				@"    #44 0x00002648e9c9  (eboot.bin+0x829c9)",
+				@"    #45 0x000026541ad9  (eboot.bin+0x135ad9)",
+				@"    #46 0x00002652a6bb  (eboot.bin+0x11e6bb)",
+				@"    #47 0x00002653ea89  (eboot.bin+0x132a89)",
+				@"    #48 0x00002652a6bb  (eboot.bin+0x11e6bb)",
+				@"    #49 0x000030e67257  (eboot.bin+0xaa5b257)",
+				@"    #50 0x000030e606e6  (eboot.bin+0xaa546e6)",
+				@"    #51 0x000030e69fea  (eboot.bin+0xaa5dfea)",
+				@"    #52 0x00002640c0bf  (eboot.bin+0xbf)",
+				@"",
+				@"Address 0x00d0572d5090 is a wild pointer inside of access range of size 0x000000000002.",
+				@"SUMMARY: AddressSanitizer: use-after-poison (eboot.bin+0x4c55d3) at (eboot.bin+0x4c55d3)",
+				@"Shadow bytes around the buggy address:",
+				@"  0x00d0572d4e00: f7 f7 f7 f7 f7 f7 f7 f7 f7 f7 f7 f7 f7 f7 f7 f7",
+				@"  0x00d0572d4e80: f7 f7 f7 f7 f7 f7 f7 f7 f7 f7 f7 f7 f7 f7 f7 f7",
+				@"  0x00d0572d4f00: f7 f7 f7 f7 f7 f7 f7 f7 f7 f7 f7 f7 f7 f7 f7 f7",
+				@"  0x00d0572d4f80: f7 f7 f7 f7 f7 f7 f7 f7 f7 f7 f7 f7 f7 f7 f7 f7",
+				@"  0x00d0572d5000: f7 f7 f7 f7 f7 f7 f7 f7 f7 f7 f7 f7 f7 f7 f7 f7",
+				@"=>0x00d0572d5080: f7 f7[f7]f7 f7 f7 f7 f7 f7 f7 f7 f7 f7 f7 f7 f7",
+				@"  0x00d0572d5100: f7 f7 f7 f7 f7 f7 f7 f7 f7 f7 f7 f7 f7 f7 f7 f7",
+				@"  0x00d0572d5180: f7 f7 f7 f7 f7 f7 f7 f7 f7 f7 f7 f7 f7 f7 f7 f7",
+				@"  0x00d0572d5200: f7 f7 f7 f7 f7 f7 f7 f7 f7 f7 f7 f7 f7 f7 f7 f7",
+				@"  0x00d0572d5280: f7 f7 f7 f7 f7 f7 f7 f7 f7 f7 f7 f7 f7 f7 f7 f7",
+				@"  0x00d0572d5300: f7 f7 f7 f7 f7 f7 f7 f7 f7 f7 f7 f7 f7 f7 f7 f7",
+				@"Shadow byte legend (one shadow byte represents 8 application bytes):",
+				@"  Addressable:           00",
+				@"  Partially addressable: 01 02 03 04 05 06 07 ",
+				@"  Heap left redzone:       fa",
+				@"  Freed heap region:       fd",
+				@"  Stack left redzone:      f1",
+				@"  Stack mid redzone:       f2",
+				@"  Stack right redzone:     f3",
+				@"  Stack after return:      f5",
+				@"  Stack use after scope:   f8",
+				@"  Global redzone:          f9",
+				@"  Global init order:       f6",
+				@"  Poisoned by user:        f7",
+				@"  Container overflow:      fc",
+				@"  Array cookie:            ac",
+				@"  Intra object redzone:    bb",
+				@"  ASan internal:           fe",
+				@"  Left alloca redzone:     ca",
+				@"  Right alloca redzone:    cb",
+				@"Fatal error: [File:Z:/UEVFS/Root/Engine/Platforms/Source/Runtime/Core/Private/PlatformMemory.cpp] [Line: 520] ",
+				@"ASan Error: =================================================================",
+				@"==85[0x18ca8]==ERROR: AddressSanitizer: use-after-poison on address 0x00d0572d5090 at pc 0x0000268d15d3 bp 0x0007effc1b90 sp 0x0007effc1b88",
+				@"READ of size 2 at 0x00d0572d5090 thread T0",
+				@"    #0 0x0000268d15d3  (eboot.bin+0x4c55d3)",
+				@"    #1 0x0000268cb2ee  (eboot.bin+0x4bf2ee)",
+				@"    #2 0x0000268d06b9  (eboot.bin+0x4c46b9)",
+				@"    #3 0x0000268d072a  (eboot.bin+0x4c472a)",
+				@"    #4 0x000026904cf7  (eboot.bin+0x4f8cf7)",
+				@"    #5 0x0000268d2fd1  (eboot.bin+0x4c6fd1)",
+				@"    #6 0x0000268d486c  (eboot.bin+0x4c886c)",
+				@"    #7 0x0000268caf3e  (eboot.bin+0x4bef3e)",
+				@"    #8 0x00002642a55f  (eboot.bin+0x1e55f)",
+				@"    #9 0x00002642a2cd  (eboot.bin+0x1e2cd)",
+				@"    #10 0x0000268ca9d9  (eboot.bin+0x4be9d9)",
+				@"    #11 0x0000267c663d  (eboot.bin+0x3ba63d)",
+				@"    #12 0x0000267433e9  (eboot.bin+0x3373e9)",
+				@"    #13 0x00004005f913  (eboot.bin+0x19c53913)",
+				@"    #14 0x00004005f7a1  (eboot.bin+0x19c537a1)",
+				@"    #15 0x000036b55372  (eboot.bin+0x10749372)",
+				@"    #16 0x000036bbc84e  (eboot.bin+0x107b084e)",
+				@"    #17 0x00002a5e38ec  (eboot.bin+0x41d78ec)",
+				@"    #18 0x000036b2bd70  (eboot.bin+0x1071fd70)",
+				@"    #19 0x000036b2b1d9  (eboot.bin+0x1071f1d9)",
+				@"    #20 0x000036bcd8e6  (eboot.bin+0x107c18e6)",
+				@"    #21 0x000036bcd55b  (eboot.bin+0x107c155b)",
+				@"    #22 0x000036b315a7  (eboot.bin+0x107255a7)",
+				@"    #23 0x000036b32697  (eboot.bin+0x10726697)",
+				@"    #24 0x000036b9dfde  (eboot.bin+0x10791fde)",
+				@"    #25 0x000036bd751c  (eboot.bin+0x107cb51c)",
+				@"    #26 0x000036bd72cb  (eboot.bin+0x107cb2cb)",
+				@"    #27 0x0000265cc963  (eboot.bin+0x1c0963)",
+				@"    #28 0x000031b40ed6  (eboot.bin+0xb734ed6)",
+				@"    #29 0x000031ac8d2e  (eboot.bin+0xb6bcd2e)",
+				@"    #30 0x000031b7c352  (eboot.bin+0xb770352)",
+				@"    #31 0x000031b7bfa0  (eboot.bin+0xb76ffa0)",
+				@"    #32 0x000029fdd12f  (eboot.bin+0x3bd112f)",
+				@"    #33 0x000031a5885e  (eboot.bin+0xb64c85e)",
+				@"    #34 0x00002a014420  (eboot.bin+0x3c08420)",
+				@"    #35 0x00002a0140d5  (eboot.bin+0x3c080d5)",
+				@"    #36 0x000029fdd12f  (eboot.bin+0x3bd112f)",
+				@"    #37 0x000029fed56c  (eboot.bin+0x3be156c)",
+				@"    #38 0x00002a014420  (eboot.bin+0x3c08420)",
+				@"    #39 0x00002a0140d5  (eboot.bin+0x3c080d5)",
+				@"    #40 0x000029fdd12f  (eboot.bin+0x3bd112f)",
+				@"    #41 0x000029fc93f0  (eboot.bin+0x3bbd3f0)",
+				@"    #42 0x000029fc902c  (eboot.bin+0x3bbd02c)",
+				@"    #43 0x000029fe3f84  (eboot.bin+0x3bd7f84)",
+				@"    #44 0x00002648e9c9  (eboot.bin+0x829c9)",
+				@"    #45 0x000026541ad9  (eboot.bin+0x135ad9)",
+				@"    #46 0x00002652a6bb  (eboot.bin+0x11e6bb)",
+				@"    #47 0x00002653ea89  (eboot.bin+0x132a89)",
+				@"    #48 0x00002652a6bb  (eboot.bin+0x11e6bb)",
+				@"    #49 0x000030e67257  (eboot.bin+0xaa5b257)",
+				@"    #50 0x000030e606e6  (eboot.bin+0xaa546e6)",
+				@"    #51 0x000030e69fea  (eboot.bin+0xaa5dfea)",
+				@"    #52 0x00002640c0bf  (eboot.bin+0xbf)",
+				@"",
+				@"Address 0x00d0572d5090 is a wild pointer inside of access range of size 0x000000000002.",
+				@"SUMMARY: AddressSanitizer: use-after-poison (eboot.bin+0x4c55d3) at (eboot.bin+0x4c55d3)",
+				@"Shadow bytes around the buggy address:",
+				@"  0x00d0572d4e00: f7 f7 f7 f7 f7 f7 f7 f7 f7 f7 f7 f7 f7 f7 f7 f7",
+				@"  0x00d0572d4e80: f7 f7 f7 f7 f7 f7 f7 f7 f7 f7 f7 f7 f7 f7 f7 f7",
+				@"  0x00d0572d4f00: f7 f7 f7 f7 f7 f7 f7 f7 f7 f7 f7 f7 f7 f7 f7 f7",
+				@"  0x00d0572d4f80: f7 f7 f7 f7 f7 f7 f7 f7 f7 f7 f7 f7 f7 f7 f7 f7",
+				@"  0x00d0572d5000: f7 f7 f7 f7 f7 f7 f7 f7 f7 f7 f7 f7 f7 f7 f7 f7",
+				@"=>0x00d0572d5080: f7 f7[f7]f7 f7 f7 f7 f7 f7 f7 f7 f7 f7 f7 f7 f7",
+				@"  0x00d0572d5100: f7 f7 f7 f7 f7 f7 f7 f7 f7 f7 f7 f7 f7 f7 f7 f7",
+				@"  0x00d0572d5180: f7 f7 f7 f7 f7 f7 f7 f7 f7 f7 f7 f7 f7 f7 f7 f7",
+				@"  0x00d0572d5200: f7 f7 f7 f7 f7 f7 f7 f7 f7 f7 f7 f7 f7 f7 f7 f7",
+				@"  0x00d0572d5280: f7 f7 f7 f7 f7 f7 f7 f7 f7 f7 f7 f7 f7 f7 f7 f7",
+				@"  0x00d0572d5300: f7 f7 f7 f7 f7 f7 f7 f7 f7 f7 f7 f7 f7 f7 f7 f7",
+				@"Shadow byte legend (one shadow byte represents 8 application bytes):",
+				@"  Addressable:           00",
+				@"  Partially addressable: 01 02 03 04 05 06 07 ",
+				@"  Heap left redzone:       fa",
+				@"  Freed heap region:       fd",
+				@"  Stack left redzone:      f1",
+				@"  Stack mid redzone:       f2",
+				@"  Stack right redzone:     f3",
+				@"  Stack after return:      f5",
+				@"  Stack use after scope:   f8",
+				@"  Global redzone:          f9",
+				@"  Global init order:       f6",
+				@"  Poisoned by user:        f7",
+				@"  Container overflow:      fc",
+				@"  Array cookie:            ac",
+				@"  Intra object redzone:    bb",
+				@"  ASan internal:           fe",
+				@"  Left alloca redzone:     ca",
+				@"  Right alloca redzone:    cb",
+				@"",
+				@"End of Address Sanitizer report",
+			};
+
+			List<LogEvent> logEvents = Parse(lines);
+			Assert.AreEqual(183, logEvents[0].LineCount);
+			CheckEventGroup(logEvents, 0, 183, LogLevel.Error, KnownLogEvents.Sanitizer_Address);
+
+			// Summary
+			Assert.AreEqual("eboot.bin+0x4c55d3", logEvents[58].GetProperty("SummarySourceFile").ToString());
+			Assert.AreEqual("eboot.bin+0x4c55d3", logEvents[58].GetProperty("Symbol").ToString());
+			Assert.AreEqual("use-after-poison", logEvents[58].GetProperty("SummaryReason").ToString());
+		}
+
 		[TestMethod]
 		public void LocalizationChannelMatcher()
 		{
diff --git a/Engine/Source/Programs/AutomationTool/AutomationUtils/Matchers/SanitizerEventMatcher.cs b/Engine/Source/Programs/AutomationTool/AutomationUtils/Matchers/SanitizerEventMatcher.cs
--- a/Engine/Source/Programs/AutomationTool/AutomationUtils/Matchers/SanitizerEventMatcher.cs
+++ b/Engine/Source/Programs/AutomationTool/AutomationUtils/Matchers/SanitizerEventMatcher.cs
@@ -31,7 +31,7 @@ namespace AutomationUtils.Matchers
 		// e.g.
 		// ASAN: ==30562==ABORTING
 		// TSAN: ==================
-		public static readonly Regex ReportEndPattern = new Regex(@"(^==\d+==(\s*\w+)?$)|(^==================$)", RegexOptions.Multiline | RegexOptions.ExplicitCapture);
+		public static readonly Regex ReportEndPattern = new Regex(@"(End of Address Sanitizer report)|(^==\d+==(\s*\w+)?$)|(^==================$)", RegexOptions.Multiline | RegexOptions.ExplicitCapture);
 
 		/// <summary>
 		/// Regex pattern that matches Report Level
@@ -39,7 +39,7 @@ namespace AutomationUtils.Matchers
 		// e.g.
 		// ASAN: ==30562==ERROR: AddressSanitizer: heap-use-after-free on address 0x617002aa8418 at pc 0x7f98a08bd090 bp 0x7ffc30203af0 sp 0x7ffc30203ae8
 		// TSAN: WARNING: ThreadSanitizer: data race (pid=6049)
-		public static readonly Regex ReportLevelPattern = new Regex(@"^[\s\t]*(==\d+==\s*)?(?<ReportLevel>WARNING|ERROR):\s+(?<SanitizerName>\w+)Sanitizer:\s*(?<Summary>.+)$", RegexOptions.Multiline | RegexOptions.ExplicitCapture);
+		public static readonly Regex ReportLevelPattern = new Regex(@"^[\s\t]*(==\d+(?:\[0x\w+\])?==\s*)?(?<ReportLevel>WARNING|ERROR):\s+(?<SanitizerName>\w+)Sanitizer:\s*(?<Summary>.+)$", RegexOptions.Multiline | RegexOptions.ExplicitCapture);
 
 		/// <summary>
 		/// Regex pattern that matches Stack Trace
@@ -65,11 +65,11 @@ namespace AutomationUtils.Matchers
 		public static readonly Regex SummaryPattern = new Regex(
 			@"^[\s\t]*SUMMARY:\s+(?<SanitizerName>\w+)Sanitizer:\s*" +
 			@"(?<SummaryReason>.+)\s" +
-			@"\(?(?<SourceFile>(\/|\w:).+?)\)?" +
-			@"(\s*\(BuildId:.*\))?" +
+			@"((\(?(?<SourceFile>(\/|\w:|eboot).+?)\)?" +
+			@"(\s+\(BuildId:\s.+\))?" +
 			@"(:(?<Line>\d+?))?" +
-			@"(:(?<Column>\d+?))?\s+?in\s+?" +
-			@"(?<Symbol>.+)$", RegexOptions.Multiline | RegexOptions.ExplicitCapture);
+			@"(:(?<Column>\d+?))?))\s+?" +
+			@"(in\s+?(?<Symbol>.+)|at\s+?\(?(?<Symbol>.+)\)?.+)$", RegexOptions.Multiline | RegexOptions.ExplicitCapture);
 
 		/// <summary>
 		/// Convert SanitizerName to EventId
diff --git a/Engine/Source/Programs/AutomationTool/Gauntlet/Unreal/Base/Gauntlet.UnrealTestNode.cs b/Engine/Source/Programs/AutomationTool/Gauntlet/Unreal/Base/Gauntlet.UnrealTestNode.cs
--- a/Engine/Source/Programs/AutomationTool/Gauntlet/Unreal/Base/Gauntlet.UnrealTestNode.cs
+++ b/Engine/Source/Programs/AutomationTool/Gauntlet/Unreal/Base/Gauntlet.UnrealTestNode.cs
@@ -2259,7 +2259,7 @@ namespace Gauntlet
 			foreach (UnrealTestEvent Event in Asserts)
 			{
 				LogEvent AssertData = DecomposeAssert(Event, MaxCallstackLines);
-				Log.Error(AssertData.Id, AssertData.Message, Args: AssertData.Properties);
+				Log.Error(AssertData.Id, AssertData.Format ?? AssertData.Message, ((Dictionary<string, object>)AssertData.Properties)?.Values.ToArray());
 				Log.Info("");
 			}
 
diff --git a/Engine/Source/Programs/AutomationTool/Gauntlet/Unreal/Utils/Gauntlet.UnrealLogParser.cs b/Engine/Source/Programs/AutomationTool/Gauntlet/Unreal/Utils/Gauntlet.UnrealLogParser.cs
--- a/Engine/Source/Programs/AutomationTool/Gauntlet/Unreal/Utils/Gauntlet.UnrealLogParser.cs
+++ b/Engine/Source/Programs/AutomationTool/Gauntlet/Unreal/Utils/Gauntlet.UnrealLogParser.cs
@@ -1027,7 +1027,6 @@ namespace Gauntlet
 			{
 				Traces = PostMortemTraces;
 			}
-
 			// Keep the one with the most information.
 			return Traces.Count() > 0 ? Traces.OrderBy(T => T.Callstack.Length).Last() : null;
 		}
@@ -1292,7 +1291,8 @@ namespace Gauntlet
 			// we'll assume that may not happen...
 
 			Regex CallstackMatch = new Regex(@"(0[xX][0-9A-f]{8,16})(?:\s+\(0[xX][0-9A-f]{8,16}\))?\s+(.+?)\s+\[(.*?)\][^\w]*$");
-			Regex ExtraErrorLine = new Regex(@".+:\s*Error:\s*");
+			Regex ExtraErrorLine = new Regex(@".+:\s*Error:\s*"); 
+			Regex TimestampMatch = new Regex(@"\[.+\]\[.+\]Log\w+\:"); 
 			foreach (UnrealLog.CallstackMessage NewTrace in Traces)
 			{
 				List<string> Backtrace = new List<string>();
@@ -1331,15 +1331,16 @@ namespace Gauntlet
 						if (Backtrace.Count == 0)
 						{
 							// Add additional summary error lines before backtrace lines are found
-							if (Line[0] != '[')
-							{ // Line with no time stamp
+							Match NLMatch = TimestampMatch.Match(Line);
+							if (!NLMatch.Success)
+							{ //New log line
 								NewTrace.Message += "\n" + Line;
 							}
 							else
-							{ // Line with error tag
+							{
 								Match MsgMatch = ExtraErrorLine.Match(Line);
 								if (MsgMatch.Success)
-								{
+								{ // Line with error tag
 									string MsgString = Line.Substring(MsgMatch.Index + MsgMatch.Length).Trim();
 									if (string.IsNullOrEmpty(MsgString)) continue;
 									NewTrace.Message += "\n" + MsgString;
