[ThreadSanitizer]

Extends LogMatcher to support TSAN summary reasons that have elaborated reasons using paretheses.

#rnx

diff --git a/Engine/Source/Programs/AutomationTool.Tests/LogParserTests.cs b/Engine/Source/Programs/AutomationTool.Tests/LogParserTests.cs
--- a/Engine/Source/Programs/AutomationTool.Tests/LogParserTests.cs
+++ b/Engine/Source/Programs/AutomationTool.Tests/LogParserTests.cs
@@ -845,6 +845,31 @@ namespace AutomationTool.Tests
 			Assert.AreEqual("data race", logEvents[SummaryIndex].GetProperty("SummaryReason").ToString());
 		}
 
+		[TestMethod]
+		public void ThreadSanitizerErrorMatcher_SummaryWithElaboratedSummaryReason()
+		{
+			string[] lines =
+			{
+				@"==================",
+				@"WARNING: ThreadSanitizer: data race (pid=66092)",
+				@"SUMMARY: ThreadSanitizer: lock-order-inversion (potential deadlock) /mnt/horde/++UE5/Sync/Engine/Source/Runtime/Core/Public/HAL/PThreadsRecursiveMutex.h:49:3 in UE::FPThreadsRecursiveMutex::Lock()",
+				@"=================="
+			};
+
+			List<LogEvent> logEvents = Parse(lines);
+			Assert.AreEqual(4, logEvents[0].LineCount);
+			CheckEventGroup(logEvents, 0, 4, LogLevel.Warning, KnownLogEvents.Sanitizer_Thread);
+
+			// Summary
+			const int SummaryIndex = 2;
+			Assert.IsTrue(logEvents[SummaryIndex].Message.StartsWith("SUMMARY: ThreadSanitizer"));
+			Assert.AreEqual("/mnt/horde/++UE5/Sync/Engine/Source/Runtime/Core/Public/HAL/PThreadsRecursiveMutex.h", logEvents[SummaryIndex].GetProperty("SummarySourceFile").ToString());
+			Assert.AreEqual("49", logEvents[SummaryIndex].GetProperty("Line").ToString());
+			Assert.AreEqual("3", logEvents[SummaryIndex].GetProperty("Column").ToString());
+			Assert.AreEqual("UE::FPThreadsRecursiveMutex::Lock()", logEvents[SummaryIndex].GetProperty("Symbol").ToString());
+			Assert.AreEqual("lock-order-inversion (potential deadlock)", logEvents[SummaryIndex].GetProperty("SummaryReason").ToString());
+		}
+
 		[TestMethod]
 		public void AddressSanitizerErrorMatcher()
 		{
diff --git a/Engine/Source/Programs/AutomationTool/AutomationUtils/Matchers/SanitizerEventMatcher.cs b/Engine/Source/Programs/AutomationTool/AutomationUtils/Matchers/SanitizerEventMatcher.cs
--- a/Engine/Source/Programs/AutomationTool/AutomationUtils/Matchers/SanitizerEventMatcher.cs
+++ b/Engine/Source/Programs/AutomationTool/AutomationUtils/Matchers/SanitizerEventMatcher.cs
@@ -64,7 +64,7 @@ namespace AutomationUtils.Matchers
 		// TSAN: SUMMARY: ThreadSanitizer: data race /some/path.cpp:169:33 in MyType::Foo()
 		public static readonly Regex SummaryPattern = new Regex(
 			@"^[\s\t]*SUMMARY:\s+(?<SanitizerName>\w+)Sanitizer:\s*" +
-			@"(?<SummaryReason>[\w\s]+)\s" +
+			@"(?<SummaryReason>[\w\s\-\(\)]+)\s" +
 			@"\(?(?<SourceFile>(\/|\w:).+?)\)?" +
 			@"(:(?<Line>\d+?))?" +
 			@"(:(?<Column>\d+?))?" +
